Remaining HoloLens platform extension files for engine plugins.

#rnx

diff --git a/Engine/Platforms/Hololens/Plugins/Runtime/AR/AzureSpatialAnchors/AzureSpatialAnchors_HoloLens.uplugin b/Engine/Platforms/Hololens/Plugins/Runtime/AR/AzureSpatialAnchors/AzureSpatialAnchors_HoloLens.uplugin
new file mode 100644
--- /dev/null
+++ b/Engine/Platforms/Hololens/Plugins/Runtime/AR/AzureSpatialAnchors/AzureSpatialAnchors_HoloLens.uplugin
@@ -0,0 +1,16 @@
+{
+	"FileVersion": 3,
+	"bIsPluginExtension": true,
+	"SupportedTargetPlatforms": [
+		"HoloLens"
+	],
+	"Modules": [
+		{
+			"Name": "AzureSpatialAnchors",
+			"Type": "Runtime",
+			"PlatformAllowList": [
+				"HoloLens"
+			]
+		}
+	]
+}
\ No newline at end of file
diff --git a/Engine/Platforms/Hololens/Plugins/Runtime/AR/AzureSpatialAnchors/Source/AzureSpatialAnchors/AzureSpatialAnchors_HoloLens.Build.cs b/Engine/Platforms/Hololens/Plugins/Runtime/AR/AzureSpatialAnchors/Source/AzureSpatialAnchors/AzureSpatialAnchors_HoloLens.Build.cs
new file mode 100644
--- /dev/null
+++ b/Engine/Platforms/Hololens/Plugins/Runtime/AR/AzureSpatialAnchors/Source/AzureSpatialAnchors/AzureSpatialAnchors_HoloLens.Build.cs
@@ -0,0 +1,11 @@
+// Copyright Epic Games, Inc. All Rights Reserved.
+
+namespace UnrealBuildTool.Rules
+{
+	public class AzureSpatialAnchors_HoloLens : AzureSpatialAnchors
+	{
+		public AzureSpatialAnchors_HoloLens(ReadOnlyTargetRules Target) : base(Target)
+		{
+		}
+	}
+}
diff --git a/Engine/Platforms/Hololens/Plugins/Runtime/MicrosoftSpatialAudio/MicroSoftSpatialSound_HoloLens.uplugin b/Engine/Platforms/Hololens/Plugins/Runtime/MicrosoftSpatialAudio/MicroSoftSpatialSound_HoloLens.uplugin
new file mode 100644
--- /dev/null
+++ b/Engine/Platforms/Hololens/Plugins/Runtime/MicrosoftSpatialAudio/MicroSoftSpatialSound_HoloLens.uplugin
@@ -0,0 +1,16 @@
+{
+	"FileVersion": 3,
+	"bIsPluginExtension": true,
+	"SupportedTargetPlatforms": [
+		"HoloLens"
+	],
+	"Modules": [
+		{
+			"Name": "MicrosoftSpatialSound",
+			"Type": "Runtime",
+			"PlatformAllowList": [
+				"HoloLens"
+			]
+		}
+	]
+}
\ No newline at end of file
diff --git a/Engine/Platforms/Hololens/Plugins/Runtime/MicrosoftSpatialAudio/Source/MicrosoftSpatialSound/MicrosoftSpatialSound_HoloLens.Build.cs b/Engine/Platforms/Hololens/Plugins/Runtime/MicrosoftSpatialAudio/Source/MicrosoftSpatialSound/MicrosoftSpatialSound_HoloLens.Build.cs
new file mode 100644
--- /dev/null
+++ b/Engine/Platforms/Hololens/Plugins/Runtime/MicrosoftSpatialAudio/Source/MicrosoftSpatialSound/MicrosoftSpatialSound_HoloLens.Build.cs
@@ -0,0 +1,11 @@
+// Copyright Epic Games, Inc. All Rights Reserved.
+
+namespace UnrealBuildTool.Rules
+{
+	public class MicrosoftSpatialSound_HoloLens : MicrosoftSpatialSound
+	{
+		public MicrosoftSpatialSound_HoloLens(ReadOnlyTargetRules Target) : base(Target)
+		{
+		}
+	}
+}
diff --git a/Engine/Platforms/Hololens/Plugins/Runtime/OpenXREyeTracker/OpenXREyeTracker_HoloLens.uplugin b/Engine/Platforms/Hololens/Plugins/Runtime/OpenXREyeTracker/OpenXREyeTracker_HoloLens.uplugin
new file mode 100644
--- /dev/null
+++ b/Engine/Platforms/Hololens/Plugins/Runtime/OpenXREyeTracker/OpenXREyeTracker_HoloLens.uplugin
@@ -0,0 +1,16 @@
+{
+	"FileVersion": 3,
+	"bIsPluginExtension": true,
+	"SupportedTargetPlatforms": [
+		"HoloLens"
+	],
+	"Modules": [
+		{
+			"Name": "OpenXREyeTracker",
+			"Type": "Runtime",
+			"PlatformAllowList": [
+				"HoloLens"
+			]
+		}
+	]
+}
\ No newline at end of file
diff --git a/Engine/Platforms/Hololens/Plugins/Runtime/OpenXREyeTracker/Source/OpenXREyeTracker/OpenXREyeTracker_HoloLens.Build.cs b/Engine/Platforms/Hololens/Plugins/Runtime/OpenXREyeTracker/Source/OpenXREyeTracker/OpenXREyeTracker_HoloLens.Build.cs
new file mode 100644
--- /dev/null
+++ b/Engine/Platforms/Hololens/Plugins/Runtime/OpenXREyeTracker/Source/OpenXREyeTracker/OpenXREyeTracker_HoloLens.Build.cs
@@ -0,0 +1,11 @@
+// Copyright Epic Games, Inc. All Rights Reserved.
+
+namespace UnrealBuildTool.Rules
+{
+	public class OpenXREyeTracker_HoloLens : OpenXREyeTracker
+	{
+		public OpenXREyeTracker_HoloLens(ReadOnlyTargetRules Target) : base(Target)
+		{
+		}
+	}
+}
diff --git a/Engine/Platforms/Hololens/Plugins/Runtime/OpenXRHandTracking/OpenXRHandTracking_HoloLens.uplugin b/Engine/Platforms/Hololens/Plugins/Runtime/OpenXRHandTracking/OpenXRHandTracking_HoloLens.uplugin
new file mode 100644
--- /dev/null
+++ b/Engine/Platforms/Hololens/Plugins/Runtime/OpenXRHandTracking/OpenXRHandTracking_HoloLens.uplugin
@@ -0,0 +1,25 @@
+{
+	"FileVersion": 3,
+	"bIsPluginExtension": true,
+	"SupportedTargetPlatforms": [
+		"HoloLens"
+	],
+	"Modules": [
+		{
+			"Name": "OpenXRHandTracking",
+			"Type": "Runtime",
+			"PlatformAllowList": [
+				"HoloLens"
+			]
+		}
+	],
+	"Plugins": [
+		{
+			"Name": "OpenXR",
+			"Enabled": true,
+			"PlatformAllowList": [
+				"HoloLens"
+			]
+		}
+	]
+}
\ No newline at end of file
diff --git a/Engine/Platforms/Hololens/Plugins/Runtime/OpenXRHandTracking/Source/OpenXRHandTracking/OpenXRHandTracking_HoloLens.Build.cs b/Engine/Platforms/Hololens/Plugins/Runtime/OpenXRHandTracking/Source/OpenXRHandTracking/OpenXRHandTracking_HoloLens.Build.cs
new file mode 100644
--- /dev/null
+++ b/Engine/Platforms/Hololens/Plugins/Runtime/OpenXRHandTracking/Source/OpenXRHandTracking/OpenXRHandTracking_HoloLens.Build.cs
@@ -0,0 +1,11 @@
+// Copyright Epic Games, Inc. All Rights Reserved.
+
+namespace UnrealBuildTool.Rules
+{
+	public class OpenXRHandTracking_HoloLens : OpenXRHandTracking
+	{
+		public OpenXRHandTracking_HoloLens(ReadOnlyTargetRules Target) : base(Target)
+		{
+		}
+	}
+}
diff --git a/Engine/Platforms/Hololens/Plugins/Runtime/OpenXRMsftHandInteraction/OpenXRMsftHandInteraction_HoloLens.uplugin b/Engine/Platforms/Hololens/Plugins/Runtime/OpenXRMsftHandInteraction/OpenXRMsftHandInteraction_HoloLens.uplugin
new file mode 100644
--- /dev/null
+++ b/Engine/Platforms/Hololens/Plugins/Runtime/OpenXRMsftHandInteraction/OpenXRMsftHandInteraction_HoloLens.uplugin
@@ -0,0 +1,25 @@
+{
+	"FileVersion": 3,
+	"bIsPluginExtension": true,
+	"SupportedTargetPlatforms": [
+		"HoloLens"
+	],
+	"Modules": [
+		{
+			"Name": "OpenXRMsftHandInteraction",
+			"Type": "Runtime",
+			"PlatformAllowList": [
+				"HoloLens"
+			]
+		}
+	],
+	"Plugins": [
+		{
+			"Name": "OpenXR",
+			"Enabled": true,
+			"PlatformAllowList": [
+				"HoloLens"
+			]
+		}
+	]
+}
\ No newline at end of file
diff --git a/Engine/Platforms/Hololens/Plugins/Runtime/OpenXRMsftHandInteraction/Source/OpenXRMsftHandInteraction/OpenXRMsftHandInteraction_HoloLens.Build.cs b/Engine/Platforms/Hololens/Plugins/Runtime/OpenXRMsftHandInteraction/Source/OpenXRMsftHandInteraction/OpenXRMsftHandInteraction_HoloLens.Build.cs
new file mode 100644
--- /dev/null
+++ b/Engine/Platforms/Hololens/Plugins/Runtime/OpenXRMsftHandInteraction/Source/OpenXRMsftHandInteraction/OpenXRMsftHandInteraction_HoloLens.Build.cs
@@ -0,0 +1,11 @@
+// Copyright Epic Games, Inc. All Rights Reserved.
+
+namespace UnrealBuildTool.Rules
+{
+	public class OpenXRMsftHandInteraction_HoloLens : OpenXRMsftHandInteraction
+	{
+		public OpenXRMsftHandInteraction_HoloLens(ReadOnlyTargetRules Target) : base(Target)
+		{
+		}
+	}
+}
diff --git a/Engine/Plugins/Runtime/AR/AzureSpatialAnchors/AzureSpatialAnchors.uplugin b/Engine/Plugins/Runtime/AR/AzureSpatialAnchors/AzureSpatialAnchors.uplugin
--- a/Engine/Plugins/Runtime/AR/AzureSpatialAnchors/AzureSpatialAnchors.uplugin
+++ b/Engine/Plugins/Runtime/AR/AzureSpatialAnchors/AzureSpatialAnchors.uplugin
@@ -17,10 +17,9 @@
 	"Installed" : false,
 	"SupportedTargetPlatforms": [
 		"Win64",
-		"HoloLens",
 		"Android",
 		"IOS",
-        "Mac"
+        	"Mac"
 	],
 	"Modules": [
 		{
@@ -28,11 +27,10 @@
 			"Type": "Runtime",
 			"PlatformAllowList": [
 				"Win64",
-				"HoloLens",
 				"Android",
 				"IOS",
-                "Mac"
+				"Mac"
 			]
 		}
 	]
-}
+}
\ No newline at end of file
diff --git a/Engine/Plugins/Runtime/MicrosoftSpatialAudio/MicroSoftSpatialSound.uplugin b/Engine/Plugins/Runtime/MicrosoftSpatialAudio/MicroSoftSpatialSound.uplugin
--- a/Engine/Plugins/Runtime/MicrosoftSpatialAudio/MicroSoftSpatialSound.uplugin
+++ b/Engine/Plugins/Runtime/MicrosoftSpatialAudio/MicroSoftSpatialSound.uplugin
@@ -14,7 +14,6 @@
 	"IsBetaVersion": false,
 	"Installed": false,
 	"SupportedTargetPlatforms": [
-		"HoloLens",
 		"Win64"
 	],
 	"Modules": [
@@ -23,7 +22,6 @@
 			"Type": "Runtime",
 			"LoadingPhase": "PreDefault",
 			"PlatformAllowList": [
-				"HoloLens",
 				"Win64"
 			]
 		}
diff --git a/Engine/Plugins/Runtime/OpenXREyeTracker/OpenXREyeTracker.uplugin b/Engine/Plugins/Runtime/OpenXREyeTracker/OpenXREyeTracker.uplugin
--- a/Engine/Plugins/Runtime/OpenXREyeTracker/OpenXREyeTracker.uplugin
+++ b/Engine/Plugins/Runtime/OpenXREyeTracker/OpenXREyeTracker.uplugin
@@ -17,7 +17,6 @@
 	"SupportedTargetPlatforms": [
 		"Win64",
 		"Linux",
-		"HoloLens",
 		"Android"
 	],
 	"Modules": [
@@ -28,7 +27,6 @@
 			"PlatformAllowList": [
 				"Win64",
 				"Linux",
-				"HoloLens",
 				"Android"
 			]
 		}
@@ -39,4 +37,4 @@
 			"Enabled": true
 		}
 	]
-}
+}
\ No newline at end of file
diff --git a/Engine/Plugins/Runtime/OpenXRHandTracking/OpenXRHandTracking.uplugin b/Engine/Plugins/Runtime/OpenXRHandTracking/OpenXRHandTracking.uplugin
--- a/Engine/Plugins/Runtime/OpenXRHandTracking/OpenXRHandTracking.uplugin
+++ b/Engine/Plugins/Runtime/OpenXRHandTracking/OpenXRHandTracking.uplugin
@@ -17,7 +17,6 @@
 	"SupportedTargetPlatforms": [
 		"Win64",
 		"Linux",
-		"HoloLens",
 		"Android"
 	],
 	"Modules": [
@@ -28,7 +27,6 @@
 			"PlatformAllowList": [
 				"Win64",
 				"Linux",
-				"HoloLens",
 				"Android"
 			]
 		}
@@ -40,7 +38,6 @@
 			"PlatformAllowList": [
 				"Win64",
 				"Linux",
-				"HoloLens",
 				"Android"
 			]
 		},
@@ -49,4 +46,4 @@
 			"Enabled": true
 		}
 	]
-}
+}
\ No newline at end of file
diff --git a/Engine/Plugins/Runtime/OpenXRMsftHandInteraction/OpenXRMsftHandInteraction.uplugin b/Engine/Plugins/Runtime/OpenXRMsftHandInteraction/OpenXRMsftHandInteraction.uplugin
--- a/Engine/Plugins/Runtime/OpenXRMsftHandInteraction/OpenXRMsftHandInteraction.uplugin
+++ b/Engine/Plugins/Runtime/OpenXRMsftHandInteraction/OpenXRMsftHandInteraction.uplugin
@@ -16,7 +16,6 @@
 	"Installed": false,
 	"SupportedTargetPlatforms": [
 		"Win64",
-		"HoloLens",
 		"Android"
 	],
 	"Modules": [
@@ -26,7 +25,6 @@
 			"LoadingPhase": "PostConfigInit",
 			"PlatformAllowList": [
 				"Win64",
-				"HoloLens",
 				"Android"
 			]
 		}
@@ -37,9 +35,8 @@
 			"Enabled": true,
 			"PlatformAllowList": [
 				"Win64",
-				"HoloLens",
 				"Android"
 			]
 		}
 	]
-}
+}
\ No newline at end of file
