OpenXRMsftHandInteraction plugin, remove beta tag
-This is no longer beta

diff --git a/Engine/Plugins/Runtime/OpenXRMsftHandInteraction/OpenXRMsftHandInteraction.uplugin b/Engine/Plugins/Runtime/OpenXRMsftHandInteraction/OpenXRMsftHandInteraction.uplugin
--- a/Engine/Plugins/Runtime/OpenXRMsftHandInteraction/OpenXRMsftHandInteraction.uplugin
+++ b/Engine/Plugins/Runtime/OpenXRMsftHandInteraction/OpenXRMsftHandInteraction.uplugin
@@ -12,7 +12,7 @@
 	"SupportURL": "",
 	"EnabledByDefault": false,
 	"CanContainContent": false,
-	"IsBetaVersion": true,
+	"IsBetaVersion": false,
 	"Installed": false,
 	"SupportedTargetPlatforms": [
 		"Win64",
