OpenXR all platform cleanup.
-Removing platform restrictions from a few XR plugins which were missed in UE-364372.
#rnx

#Synced-CL 50546487

diff --git a/Engine/Plugins/Runtime/OpenXRMsftHandInteraction/OpenXRMsftHandInteraction.uplugin b/Engine/Plugins/Runtime/OpenXRMsftHandInteraction/OpenXRMsftHandInteraction.uplugin
--- a/Engine/Plugins/Runtime/OpenXRMsftHandInteraction/OpenXRMsftHandInteraction.uplugin
+++ b/Engine/Plugins/Runtime/OpenXRMsftHandInteraction/OpenXRMsftHandInteraction.uplugin
@@ -14,29 +14,17 @@
 	"CanContainContent": false,
 	"IsBetaVersion": false,
 	"Installed": false,
-	"SupportedTargetPlatforms": [
-		"Win64",
-		"Android"
-	],
 	"Modules": [
 		{
 			"Name": "OpenXRMsftHandInteraction",
 			"Type": "Runtime",
-			"LoadingPhase": "PostConfigInit",
-			"PlatformAllowList": [
-				"Win64",
-				"Android"
-			]
+			"LoadingPhase": "PostConfigInit"
 		}
 	],
 	"Plugins": [
 		{
 			"Name": "OpenXR",
-			"Enabled": true,
-			"PlatformAllowList": [
-				"Win64",
-				"Android"
-			]
+			"Enabled": true
 		}
 	]
 }
\ No newline at end of file
diff --git a/Engine/Plugins/Runtime/OpenXRViveTracker/OpenXRViveTracker.uplugin b/Engine/Plugins/Runtime/OpenXRViveTracker/OpenXRViveTracker.uplugin
--- a/Engine/Plugins/Runtime/OpenXRViveTracker/OpenXRViveTracker.uplugin
+++ b/Engine/Plugins/Runtime/OpenXRViveTracker/OpenXRViveTracker.uplugin
@@ -34,12 +34,7 @@
 	"Plugins": [
 		{
 			"Name": "OpenXR",
-			"Enabled": true,
-			"PlatformAllowList": [
-				"Win64",
-				"Linux",
-				"Android"
-			]
+			"Enabled": true
 		}
 	]
 }
\ No newline at end of file
